Fix critical section bug where the audio mixer and the game thread can become stuck/deadlock on 2 mutexes triggered in reverse order.

Stopping a video recording via console command or via API can cause indefinite hang. Repo rate is 1/40.

#rnx

Deadlock occurs with MainThread and Audio FMixerDevice.

MainThread:

HighlightRecorder::Stop()

GameMediaEncorder::UnregisterListener (Locks Listener Mutex)

......

FMixerDevice::UnregisterSubmixBufferListener

FMixerSubmix::UnregisterBufferListener (Locks BufferListenerCriticalSection)

FMixerDevice:

FMixerDevice::ONProcessAudioStream

FMixerSubmix:: ProcessAudio (Locks BufferListenerCriticalSection)

......

FGameplayMediaEncoder::OnMediaSampleReady (Locks Listener Mutex)

diff --git a/Engine/Source/Runtime/GameplayMediaEncoder/Private/GameplayMediaEncoder.cpp b/Engine/Source/Runtime/GameplayMediaEncoder/Private/GameplayMediaEncoder.cpp
--- a/Engine/Source/Runtime/GameplayMediaEncoder/Private/GameplayMediaEncoder.cpp
+++ b/Engine/Source/Runtime/GameplayMediaEncoder/Private/GameplayMediaEncoder.cpp
@@ -124,10 +124,16 @@ bool FGameplayMediaEncoder::RegisterListener(IGameplayMediaEncoderListener* List
 void FGameplayMediaEncoder::UnregisterListener(IGameplayMediaEncoderListener* Listener)
 {
 	check(IsInGameThread());
-	FScopeLock Lock(&ListenersCS);
+
+	ListenersCS.Lock();
 
 	Listeners.Remove(Listener);
-	if (Listeners.Num() == 0)
+
+	bool bAnyListenersLeft = Listeners.Num() > 0;
+
+	ListenersCS.Unlock();
+
+	if (bAnyListenersLeft == false)
 	{
 		UE_LOG(GameplayMediaEncoder, Log, TEXT("Unregistered the last listener"));
 		Stop();
