Enable media capture module in PIE for use in automation.

diff --git a/Engine/Source/Runtime/GameplayMediaEncoder/Private/GameplayMediaEncoder.cpp b/Engine/Source/Runtime/GameplayMediaEncoder/Private/GameplayMediaEncoder.cpp
--- a/Engine/Source/Runtime/GameplayMediaEncoder/Private/GameplayMediaEncoder.cpp
+++ b/Engine/Source/Runtime/GameplayMediaEncoder/Private/GameplayMediaEncoder.cpp
@@ -288,24 +288,19 @@ bool FGameplayMediaEncoder::Start()
 	//
 	// subscribe to engine delegates for audio output and back buffer
 	//
-	if (UEngine* Engine = GEngine)
-	{
-		FAudioDevice* AudioDevice = Engine->GetMainAudioDevice();
-		if (AudioDevice)
-		{
-			AudioDevice->RegisterSubmixBufferListener(this);
-		}
-
-		VideoEncoder->Start();
-		FSlateApplication::Get().GetRenderer()->OnBackBufferReadyToPresent().AddRaw(this, &FGameplayMediaEncoder::OnBackBufferReady);
 
-		return true;
-	}
-	else
+	FAudioDevice* AudioDevice = GEngine->GetMainAudioDevice();
+	if (AudioDevice)
 	{
-		UE_LOG(GameplayMediaEncoder, Error, TEXT("Failed to register audio and/or video callbacks"));
-		return false;
+		AudioDevice->RegisterSubmixBufferListener(this);
 	}
+
+	VideoEncoder->Start();
+	FSlateRenderer::FOnBackBufferReadyToPresent OnBackBufferReadyDelegate;
+	OnBackBufferReadyDelegate.AddRaw(this, &FGameplayMediaEncoder::OnBackBufferReady);
+	FSlateApplication::Get().GetRenderer()->OnBackBufferReadyToPresent() = OnBackBufferReadyDelegate;
+
+	return true;
 }
 
 void FGameplayMediaEncoder::Stop()
