UE4 - Workaround deprecation warnings
#rnx

diff --git a/Engine/Source/Runtime/GameplayMediaEncoder/Private/Microsoft/Windows/AmdAmfVideoEncoder.cpp b/Engine/Source/Runtime/GameplayMediaEncoder/Private/Microsoft/Windows/AmdAmfVideoEncoder.cpp
--- a/Engine/Source/Runtime/GameplayMediaEncoder/Private/Microsoft/Windows/AmdAmfVideoEncoder.cpp
+++ b/Engine/Source/Runtime/GameplayMediaEncoder/Private/Microsoft/Windows/AmdAmfVideoEncoder.cpp
@@ -365,7 +365,9 @@ void FAmdAmfVideoEncoder::ResolveBackBuffer(const FTexture2DRHIRef& BackBuffer,
 	}
 	else // Texture format mismatch, use a shader to do the copy.
 	{
+		PRAGMA_DISABLE_DEPRECATION_WARNINGS
 		SetRenderTarget(RHICmdList, ResolvedBackBuffer, FTextureRHIRef());
+		PRAGMA_ENABLE_DEPRECATION_WARNINGS
 		RHICmdList.SetViewport(0, 0, 0.0f, ResolvedBackBuffer->GetSizeX(), ResolvedBackBuffer->GetSizeY(), 1.0f);
 
 		FGraphicsPipelineStateInitializer GraphicsPSOInit;
@@ -402,7 +404,7 @@ void FAmdAmfVideoEncoder::ResolveBackBuffer(const FTexture2DRHIRef& BackBuffer,
 			*VertexShader,
 			EDRF_Default);
 
-		SetRenderTarget(RHICmdList, FTextureRHIRef(), FTextureRHIRef());
+		UnbindRenderTargets(RHICmdList);
 	}
 }
 
diff --git a/Engine/Source/Runtime/GameplayMediaEncoder/Private/Microsoft/Windows/NvVideoEncoder.cpp b/Engine/Source/Runtime/GameplayMediaEncoder/Private/Microsoft/Windows/NvVideoEncoder.cpp
--- a/Engine/Source/Runtime/GameplayMediaEncoder/Private/Microsoft/Windows/NvVideoEncoder.cpp
+++ b/Engine/Source/Runtime/GameplayMediaEncoder/Private/Microsoft/Windows/NvVideoEncoder.cpp
@@ -334,7 +334,9 @@ void FNvVideoEncoder::CopyBackBuffer(const FTexture2DRHIRef& SrcBackBuffer, cons
 	}
 	else // Texture format mismatch, use a shader to do the copy.
 	{
+		PRAGMA_DISABLE_DEPRECATION_WARNINGS
 		SetRenderTarget(RHICmdList, DstFrame.ResolvedBackBuffer, FTextureRHIRef());
+		PRAGMA_ENABLE_DEPRECATION_WARNINGS
 		RHICmdList.SetViewport(0, 0, 0.0f, DstFrame.ResolvedBackBuffer->GetSizeX(), DstFrame.ResolvedBackBuffer->GetSizeY(), 1.0f);
 	
 		FGraphicsPipelineStateInitializer GraphicsPSOInit;
diff --git a/Engine/Source/Runtime/GameplayMediaEncoder/Private/Microsoft/Windows/WmfVideoEncoder.cpp b/Engine/Source/Runtime/GameplayMediaEncoder/Private/Microsoft/Windows/WmfVideoEncoder.cpp
--- a/Engine/Source/Runtime/GameplayMediaEncoder/Private/Microsoft/Windows/WmfVideoEncoder.cpp
+++ b/Engine/Source/Runtime/GameplayMediaEncoder/Private/Microsoft/Windows/WmfVideoEncoder.cpp
@@ -480,7 +480,9 @@ void FWmfVideoEncoder::ResolveBackBuffer(const FTexture2DRHIRef& BackBuffer, con
 	}
 	else // Texture format mismatch, use a shader to do the copy.
 	{
+		PRAGMA_DISABLE_DEPRECATION_WARNINGS
 		SetRenderTarget(RHICmdList, ResolvedBackBuffer, FTextureRHIRef());
+		PRAGMA_ENABLE_DEPRECATION_WARNINGS
 		RHICmdList.SetViewport(0, 0, 0.0f, ResolvedBackBuffer->GetSizeX(), ResolvedBackBuffer->GetSizeY(), 1.0f);
 
 		FGraphicsPipelineStateInitializer GraphicsPSOInit;
