Fix copyright notices.

diff --git a/Engine/Source/Runtime/ApplicationCore/Private/IOS/IOSFeedbackContext.cpp b/Engine/Source/Runtime/ApplicationCore/Private/IOS/IOSFeedbackContext.cpp
--- a/Engine/Source/Runtime/ApplicationCore/Private/IOS/IOSFeedbackContext.cpp
+++ b/Engine/Source/Runtime/ApplicationCore/Private/IOS/IOSFeedbackContext.cpp
@@ -1,4 +1,4 @@
-// Copyright 1998-2018 Epic Games, Inc. All Rights Reserved.
+// Copyright 1998-2019 Epic Games, Inc. All Rights Reserved.
 
 #include "IOS/IOSFeedbackContext.h"
 
diff --git a/Engine/Source/Runtime/ApplicationCore/Public/IOS/IOSFeedbackContext.h b/Engine/Source/Runtime/ApplicationCore/Public/IOS/IOSFeedbackContext.h
--- a/Engine/Source/Runtime/ApplicationCore/Public/IOS/IOSFeedbackContext.h
+++ b/Engine/Source/Runtime/ApplicationCore/Public/IOS/IOSFeedbackContext.h
@@ -1,4 +1,4 @@
-// Copyright 1998-2018 Epic Games, Inc. All Rights Reserved.
+// Copyright 1998-2019 Epic Games, Inc. All Rights Reserved.
 
 #pragma once
 #include "Misc/App.h"
diff --git a/Engine/Source/Runtime/AudioMixer/Classes/Generators/AudioGenerator.h b/Engine/Source/Runtime/AudioMixer/Classes/Generators/AudioGenerator.h
--- a/Engine/Source/Runtime/AudioMixer/Classes/Generators/AudioGenerator.h
+++ b/Engine/Source/Runtime/AudioMixer/Classes/Generators/AudioGenerator.h
@@ -1,4 +1,4 @@
-// Copyright 1998-2018 Epic Games, Inc. All Rights Reserved.
+// Copyright 1998-2019 Epic Games, Inc. All Rights Reserved.
 
 #pragma once
 
diff --git a/Engine/Source/Runtime/AudioMixer/Private/Generators/AudioGenerator.cpp b/Engine/Source/Runtime/AudioMixer/Private/Generators/AudioGenerator.cpp
--- a/Engine/Source/Runtime/AudioMixer/Private/Generators/AudioGenerator.cpp
+++ b/Engine/Source/Runtime/AudioMixer/Private/Generators/AudioGenerator.cpp
@@ -1,4 +1,4 @@
-// Copyright 1998-2018 Epic Games, Inc. All Rights Reserved.
+// Copyright 1998-2019 Epic Games, Inc. All Rights Reserved.
 
 #include "Generators/AudioGenerator.h"
 
diff --git a/Engine/Source/Runtime/Core/Public/Algo/FindSequence.h b/Engine/Source/Runtime/Core/Public/Algo/FindSequence.h
--- a/Engine/Source/Runtime/Core/Public/Algo/FindSequence.h
+++ b/Engine/Source/Runtime/Core/Public/Algo/FindSequence.h
@@ -1,4 +1,4 @@
-// Copyright 1998-2018 Epic Games, Inc. All Rights Reserved.
+// Copyright 1998-2019 Epic Games, Inc. All Rights Reserved.
 
 #pragma once
 
diff --git a/Engine/Source/Runtime/Core/Public/HAL/LowLevelMemStats.h b/Engine/Source/Runtime/Core/Public/HAL/LowLevelMemStats.h
--- a/Engine/Source/Runtime/Core/Public/HAL/LowLevelMemStats.h
+++ b/Engine/Source/Runtime/Core/Public/HAL/LowLevelMemStats.h
@@ -1,4 +1,4 @@
-// Copyright 1998-2018 Epic Games, Inc. All Rights Reserved.
+// Copyright 1998-2019 Epic Games, Inc. All Rights Reserved.
 
 #include "HAL/LowLevelMemTracker.h"
 #include "Stats/Stats.h"
diff --git a/Engine/Source/Runtime/GameplayMediaEncoder/GameplayMediaEncoder.Build.cs b/Engine/Source/Runtime/GameplayMediaEncoder/GameplayMediaEncoder.Build.cs
--- a/Engine/Source/Runtime/GameplayMediaEncoder/GameplayMediaEncoder.Build.cs
+++ b/Engine/Source/Runtime/GameplayMediaEncoder/GameplayMediaEncoder.Build.cs
@@ -1,4 +1,4 @@
-// Copyright 1998-2018 Epic Games, Inc. All Rights Reserved.
+// Copyright 1998-2019 Epic Games, Inc. All Rights Reserved.
 
 using System.IO;
 using UnrealBuildTool;
diff --git a/Engine/Source/Runtime/GameplayMediaEncoder/Private/BaseVideoEncoder.cpp b/Engine/Source/Runtime/GameplayMediaEncoder/Private/BaseVideoEncoder.cpp
--- a/Engine/Source/Runtime/GameplayMediaEncoder/Private/BaseVideoEncoder.cpp
+++ b/Engine/Source/Runtime/GameplayMediaEncoder/Private/BaseVideoEncoder.cpp
@@ -1,4 +1,4 @@
-// Copyright 1998-2018 Epic Games, Inc. All Rights Reserved.
+// Copyright 1998-2019 Epic Games, Inc. All Rights Reserved.
 
 #include "BaseVideoEncoder.h"
 
diff --git a/Engine/Source/Runtime/GameplayMediaEncoder/Private/BaseVideoEncoder.h b/Engine/Source/Runtime/GameplayMediaEncoder/Private/BaseVideoEncoder.h
--- a/Engine/Source/Runtime/GameplayMediaEncoder/Private/BaseVideoEncoder.h
+++ b/Engine/Source/Runtime/GameplayMediaEncoder/Private/BaseVideoEncoder.h
@@ -1,4 +1,4 @@
-// Copyright 1998-2018 Epic Games, Inc. All Rights Reserved.
+// Copyright 1998-2019 Epic Games, Inc. All Rights Reserved.
 
 #pragma once
 
diff --git a/Engine/Source/Runtime/GameplayMediaEncoder/Private/GameplayMediaEncoder.cpp b/Engine/Source/Runtime/GameplayMediaEncoder/Private/GameplayMediaEncoder.cpp
--- a/Engine/Source/Runtime/GameplayMediaEncoder/Private/GameplayMediaEncoder.cpp
+++ b/Engine/Source/Runtime/GameplayMediaEncoder/Private/GameplayMediaEncoder.cpp
@@ -1,4 +1,4 @@
-// Copyright 1998-2018 Epic Games, Inc. All Rights Reserved.
+// Copyright 1998-2019 Epic Games, Inc. All Rights Reserved.
 
 #include "GameplayMediaEncoder.h"
 #include "Engine/GameEngine.h"
diff --git a/Engine/Source/Runtime/GameplayMediaEncoder/Private/GameplayMediaEncoderCommon.cpp b/Engine/Source/Runtime/GameplayMediaEncoder/Private/GameplayMediaEncoderCommon.cpp
--- a/Engine/Source/Runtime/GameplayMediaEncoder/Private/GameplayMediaEncoderCommon.cpp
+++ b/Engine/Source/Runtime/GameplayMediaEncoder/Private/GameplayMediaEncoderCommon.cpp
@@ -1,4 +1,4 @@
-// Copyright 1998-2018 Epic Games, Inc. All Rights Reserved.
+// Copyright 1998-2019 Epic Games, Inc. All Rights Reserved.
 
 #include "GameplayMediaEncoderCommon.h"
 #include "RHI.h"
diff --git a/Engine/Source/Runtime/GameplayMediaEncoder/Private/GameplayMediaEncoderCommon.h b/Engine/Source/Runtime/GameplayMediaEncoder/Private/GameplayMediaEncoderCommon.h
--- a/Engine/Source/Runtime/GameplayMediaEncoder/Private/GameplayMediaEncoderCommon.h
+++ b/Engine/Source/Runtime/GameplayMediaEncoder/Private/GameplayMediaEncoderCommon.h
@@ -1,4 +1,4 @@
-// Copyright 1998-2018 Epic Games, Inc. All Rights Reserved.
+// Copyright 1998-2019 Epic Games, Inc. All Rights Reserved.
 
 #pragma once
 
diff --git a/Engine/Source/Runtime/GameplayMediaEncoder/Private/GameplayMediaEncoderModule.cpp b/Engine/Source/Runtime/GameplayMediaEncoder/Private/GameplayMediaEncoderModule.cpp
--- a/Engine/Source/Runtime/GameplayMediaEncoder/Private/GameplayMediaEncoderModule.cpp
+++ b/Engine/Source/Runtime/GameplayMediaEncoder/Private/GameplayMediaEncoderModule.cpp
@@ -1,4 +1,4 @@
-// Copyright 1998-2018 Epic Games, Inc. All Rights Reserved.
+// Copyright 1998-2019 Epic Games, Inc. All Rights Reserved.
 
 #include "CoreMinimal.h"
 #include "Modules/ModuleInterface.h"
diff --git a/Engine/Source/Runtime/GameplayMediaEncoder/Private/GameplayMediaEncoderSample.cpp b/Engine/Source/Runtime/GameplayMediaEncoder/Private/GameplayMediaEncoderSample.cpp
--- a/Engine/Source/Runtime/GameplayMediaEncoder/Private/GameplayMediaEncoderSample.cpp
+++ b/Engine/Source/Runtime/GameplayMediaEncoder/Private/GameplayMediaEncoderSample.cpp
@@ -1,4 +1,4 @@
-// Copyright 1998-2018 Epic Games, Inc. All Rights Reserved.
+// Copyright 1998-2019 Epic Games, Inc. All Rights Reserved.
 
 #include "GameplayMediaEncoderSample.h"
 #include "GameplayMediaEncoderCommon.h"
diff --git a/Engine/Source/Runtime/GameplayMediaEncoder/Private/IbmLiveStreaming.cpp b/Engine/Source/Runtime/GameplayMediaEncoder/Private/IbmLiveStreaming.cpp
--- a/Engine/Source/Runtime/GameplayMediaEncoder/Private/IbmLiveStreaming.cpp
+++ b/Engine/Source/Runtime/GameplayMediaEncoder/Private/IbmLiveStreaming.cpp
@@ -1,4 +1,4 @@
-// Copyright 1998-2018 Epic Games, Inc. All Rights Reserved.
+// Copyright 1998-2019 Epic Games, Inc. All Rights Reserved.
 
 #include "IbmLiveStreaming.h"
 
diff --git a/Engine/Source/Runtime/GameplayMediaEncoder/Private/IbmLiveStreaming.h b/Engine/Source/Runtime/GameplayMediaEncoder/Private/IbmLiveStreaming.h
--- a/Engine/Source/Runtime/GameplayMediaEncoder/Private/IbmLiveStreaming.h
+++ b/Engine/Source/Runtime/GameplayMediaEncoder/Private/IbmLiveStreaming.h
@@ -1,4 +1,4 @@
-// Copyright 1998-2018 Epic Games, Inc. All Rights Reserved.
+// Copyright 1998-2019 Epic Games, Inc. All Rights Reserved.
 
 #pragma once
 
diff --git a/Engine/Source/Runtime/GameplayMediaEncoder/Private/Microsoft/Windows/AmdAmfPrivate.h b/Engine/Source/Runtime/GameplayMediaEncoder/Private/Microsoft/Windows/AmdAmfPrivate.h
--- a/Engine/Source/Runtime/GameplayMediaEncoder/Private/Microsoft/Windows/AmdAmfPrivate.h
+++ b/Engine/Source/Runtime/GameplayMediaEncoder/Private/Microsoft/Windows/AmdAmfPrivate.h
@@ -1,4 +1,4 @@
-// Copyright 1998-2018 Epic Games, Inc. All Rights Reserved.
+// Copyright 1998-2019 Epic Games, Inc. All Rights Reserved.
 
 #pragma once
 
diff --git a/Engine/Source/Runtime/GameplayMediaEncoder/Private/Microsoft/Windows/AmdAmfVideoEncoder.cpp b/Engine/Source/Runtime/GameplayMediaEncoder/Private/Microsoft/Windows/AmdAmfVideoEncoder.cpp
--- a/Engine/Source/Runtime/GameplayMediaEncoder/Private/Microsoft/Windows/AmdAmfVideoEncoder.cpp
+++ b/Engine/Source/Runtime/GameplayMediaEncoder/Private/Microsoft/Windows/AmdAmfVideoEncoder.cpp
@@ -1,4 +1,4 @@
-// Copyright 1998-2018 Epic Games, Inc. All Rights Reserved.
+// Copyright 1998-2019 Epic Games, Inc. All Rights Reserved.
 
 #include "AmdAmfVideoEncoder.h"
 
diff --git a/Engine/Source/Runtime/GameplayMediaEncoder/Private/Microsoft/Windows/AmdAmfVideoEncoder.h b/Engine/Source/Runtime/GameplayMediaEncoder/Private/Microsoft/Windows/AmdAmfVideoEncoder.h
--- a/Engine/Source/Runtime/GameplayMediaEncoder/Private/Microsoft/Windows/AmdAmfVideoEncoder.h
+++ b/Engine/Source/Runtime/GameplayMediaEncoder/Private/Microsoft/Windows/AmdAmfVideoEncoder.h
@@ -1,4 +1,4 @@
-// Copyright 1998-2018 Epic Games, Inc. All Rights Reserved.
+// Copyright 1998-2019 Epic Games, Inc. All Rights Reserved.
 
 #pragma once
 
diff --git a/Engine/Source/Runtime/GameplayMediaEncoder/Private/Microsoft/Windows/D3D11VideoProcessor.cpp b/Engine/Source/Runtime/GameplayMediaEncoder/Private/Microsoft/Windows/D3D11VideoProcessor.cpp
--- a/Engine/Source/Runtime/GameplayMediaEncoder/Private/Microsoft/Windows/D3D11VideoProcessor.cpp
+++ b/Engine/Source/Runtime/GameplayMediaEncoder/Private/Microsoft/Windows/D3D11VideoProcessor.cpp
@@ -1,4 +1,4 @@
-// Copyright 1998-2018 Epic Games, Inc. All Rights Reserved.
+// Copyright 1998-2019 Epic Games, Inc. All Rights Reserved.
 
 #include "D3D11VideoProcessor.h"
 
diff --git a/Engine/Source/Runtime/GameplayMediaEncoder/Private/Microsoft/Windows/D3D11VideoProcessor.h b/Engine/Source/Runtime/GameplayMediaEncoder/Private/Microsoft/Windows/D3D11VideoProcessor.h
--- a/Engine/Source/Runtime/GameplayMediaEncoder/Private/Microsoft/Windows/D3D11VideoProcessor.h
+++ b/Engine/Source/Runtime/GameplayMediaEncoder/Private/Microsoft/Windows/D3D11VideoProcessor.h
@@ -1,4 +1,4 @@
-// Copyright 1998-2018 Epic Games, Inc. All Rights Reserved.
+// Copyright 1998-2019 Epic Games, Inc. All Rights Reserved.
 
 #pragma once
 
diff --git a/Engine/Source/Runtime/GameplayMediaEncoder/Private/Microsoft/Windows/EncoderDevice.cpp b/Engine/Source/Runtime/GameplayMediaEncoder/Private/Microsoft/Windows/EncoderDevice.cpp
--- a/Engine/Source/Runtime/GameplayMediaEncoder/Private/Microsoft/Windows/EncoderDevice.cpp
+++ b/Engine/Source/Runtime/GameplayMediaEncoder/Private/Microsoft/Windows/EncoderDevice.cpp
@@ -1,4 +1,4 @@
-// Copyright 1998-2018 Epic Games, Inc. All Rights Reserved.
+// Copyright 1998-2019 Epic Games, Inc. All Rights Reserved.
 
 #include "EncoderDevice.h"
 
diff --git a/Engine/Source/Runtime/GameplayMediaEncoder/Private/Microsoft/Windows/EncoderDevice.h b/Engine/Source/Runtime/GameplayMediaEncoder/Private/Microsoft/Windows/EncoderDevice.h
--- a/Engine/Source/Runtime/GameplayMediaEncoder/Private/Microsoft/Windows/EncoderDevice.h
+++ b/Engine/Source/Runtime/GameplayMediaEncoder/Private/Microsoft/Windows/EncoderDevice.h
@@ -1,4 +1,4 @@
-// Copyright 1998-2018 Epic Games, Inc. All Rights Reserved.
+// Copyright 1998-2019 Epic Games, Inc. All Rights Reserved.
 
 #pragma once
 
diff --git a/Engine/Source/Runtime/GameplayMediaEncoder/Private/Microsoft/Windows/NvVideoEncoder.cpp b/Engine/Source/Runtime/GameplayMediaEncoder/Private/Microsoft/Windows/NvVideoEncoder.cpp
--- a/Engine/Source/Runtime/GameplayMediaEncoder/Private/Microsoft/Windows/NvVideoEncoder.cpp
+++ b/Engine/Source/Runtime/GameplayMediaEncoder/Private/Microsoft/Windows/NvVideoEncoder.cpp
@@ -1,4 +1,4 @@
-// Copyright 1998-2018 Epic Games, Inc. All Rights Reserved.
+// Copyright 1998-2019 Epic Games, Inc. All Rights Reserved.
 
 #include "NvVideoEncoder.h"
 
diff --git a/Engine/Source/Runtime/GameplayMediaEncoder/Private/Microsoft/Windows/NvVideoEncoder.h b/Engine/Source/Runtime/GameplayMediaEncoder/Private/Microsoft/Windows/NvVideoEncoder.h
--- a/Engine/Source/Runtime/GameplayMediaEncoder/Private/Microsoft/Windows/NvVideoEncoder.h
+++ b/Engine/Source/Runtime/GameplayMediaEncoder/Private/Microsoft/Windows/NvVideoEncoder.h
@@ -1,4 +1,4 @@
-// Copyright 1998-2018 Epic Games, Inc. All Rights Reserved.
+// Copyright 1998-2019 Epic Games, Inc. All Rights Reserved.
 
 #pragma once
 
diff --git a/Engine/Source/Runtime/GameplayMediaEncoder/Private/Microsoft/Windows/UtilityShaders.cpp b/Engine/Source/Runtime/GameplayMediaEncoder/Private/Microsoft/Windows/UtilityShaders.cpp
--- a/Engine/Source/Runtime/GameplayMediaEncoder/Private/Microsoft/Windows/UtilityShaders.cpp
+++ b/Engine/Source/Runtime/GameplayMediaEncoder/Private/Microsoft/Windows/UtilityShaders.cpp
@@ -1,4 +1,4 @@
-// Copyright 1998-2018 Epic Games, Inc. All Rights Reserved.
+// Copyright 1998-2019 Epic Games, Inc. All Rights Reserved.
 
 /*=============================================================================
 	UtilityShaders.cpp: Utility shaders for Windows Media Foundation pipeline.
diff --git a/Engine/Source/Runtime/GameplayMediaEncoder/Private/Microsoft/Windows/UtilityShaders.h b/Engine/Source/Runtime/GameplayMediaEncoder/Private/Microsoft/Windows/UtilityShaders.h
--- a/Engine/Source/Runtime/GameplayMediaEncoder/Private/Microsoft/Windows/UtilityShaders.h
+++ b/Engine/Source/Runtime/GameplayMediaEncoder/Private/Microsoft/Windows/UtilityShaders.h
@@ -1,4 +1,4 @@
-// Copyright 1998-2018 Epic Games, Inc. All Rights Reserved.
+// Copyright 1998-2019 Epic Games, Inc. All Rights Reserved.
 
 /*=============================================================================
 	UtilityShaders.h: Screen rendering definitions.
diff --git a/Engine/Source/Runtime/GameplayMediaEncoder/Private/Microsoft/Windows/WmfVideoEncoder.cpp b/Engine/Source/Runtime/GameplayMediaEncoder/Private/Microsoft/Windows/WmfVideoEncoder.cpp
--- a/Engine/Source/Runtime/GameplayMediaEncoder/Private/Microsoft/Windows/WmfVideoEncoder.cpp
+++ b/Engine/Source/Runtime/GameplayMediaEncoder/Private/Microsoft/Windows/WmfVideoEncoder.cpp
@@ -1,4 +1,4 @@
-// Copyright 1998-2018 Epic Games, Inc. All Rights Reserved.
+// Copyright 1998-2019 Epic Games, Inc. All Rights Reserved.
 
 #include "WmfVideoEncoder.h"
 
diff --git a/Engine/Source/Runtime/GameplayMediaEncoder/Private/Microsoft/Windows/WmfVideoEncoder.h b/Engine/Source/Runtime/GameplayMediaEncoder/Private/Microsoft/Windows/WmfVideoEncoder.h
--- a/Engine/Source/Runtime/GameplayMediaEncoder/Private/Microsoft/Windows/WmfVideoEncoder.h
+++ b/Engine/Source/Runtime/GameplayMediaEncoder/Private/Microsoft/Windows/WmfVideoEncoder.h
@@ -1,4 +1,4 @@
-// Copyright 1998-2018 Epic Games, Inc. All Rights Reserved.
+// Copyright 1998-2019 Epic Games, Inc. All Rights Reserved.
 
 #pragma once
 
diff --git a/Engine/Source/Runtime/GameplayMediaEncoder/Private/Microsoft/WmfAudioEncoder.cpp b/Engine/Source/Runtime/GameplayMediaEncoder/Private/Microsoft/WmfAudioEncoder.cpp
--- a/Engine/Source/Runtime/GameplayMediaEncoder/Private/Microsoft/WmfAudioEncoder.cpp
+++ b/Engine/Source/Runtime/GameplayMediaEncoder/Private/Microsoft/WmfAudioEncoder.cpp
@@ -1,4 +1,4 @@
-// Copyright 1998-2018 Epic Games, Inc. All Rights Reserved.
+// Copyright 1998-2019 Epic Games, Inc. All Rights Reserved.
 
 #include "WmfAudioEncoder.h"
 #include "GameplayMediaEncoderSample.h"
diff --git a/Engine/Source/Runtime/GameplayMediaEncoder/Private/Microsoft/WmfAudioEncoder.h b/Engine/Source/Runtime/GameplayMediaEncoder/Private/Microsoft/WmfAudioEncoder.h
--- a/Engine/Source/Runtime/GameplayMediaEncoder/Private/Microsoft/WmfAudioEncoder.h
+++ b/Engine/Source/Runtime/GameplayMediaEncoder/Private/Microsoft/WmfAudioEncoder.h
@@ -1,4 +1,4 @@
-// Copyright 1998-2018 Epic Games, Inc. All Rights Reserved.
+// Copyright 1998-2019 Epic Games, Inc. All Rights Reserved.
 
 #pragma once
 
diff --git a/Engine/Source/Runtime/GameplayMediaEncoder/Public/GameplayMediaEncoder.h b/Engine/Source/Runtime/GameplayMediaEncoder/Public/GameplayMediaEncoder.h
--- a/Engine/Source/Runtime/GameplayMediaEncoder/Public/GameplayMediaEncoder.h
+++ b/Engine/Source/Runtime/GameplayMediaEncoder/Public/GameplayMediaEncoder.h
@@ -1,4 +1,4 @@
-// Copyright 1998-2018 Epic Games, Inc. All Rights Reserved.
+// Copyright 1998-2019 Epic Games, Inc. All Rights Reserved.
 
 #pragma once
 
diff --git a/Engine/Source/Runtime/GameplayMediaEncoder/Public/GameplayMediaEncoderSample.h b/Engine/Source/Runtime/GameplayMediaEncoder/Public/GameplayMediaEncoderSample.h
--- a/Engine/Source/Runtime/GameplayMediaEncoder/Public/GameplayMediaEncoderSample.h
+++ b/Engine/Source/Runtime/GameplayMediaEncoder/Public/GameplayMediaEncoderSample.h
@@ -1,4 +1,4 @@
-// Copyright 1998-2018 Epic Games, Inc. All Rights Reserved.
+// Copyright 1998-2019 Epic Games, Inc. All Rights Reserved.
 
 #pragma once
 
diff --git a/Engine/Source/Runtime/WebBrowser/Private/Native/NativeJSScripting.cpp b/Engine/Source/Runtime/WebBrowser/Private/Native/NativeJSScripting.cpp
--- a/Engine/Source/Runtime/WebBrowser/Private/Native/NativeJSScripting.cpp
+++ b/Engine/Source/Runtime/WebBrowser/Private/Native/NativeJSScripting.cpp
@@ -1,4 +1,4 @@
-// Copyright 1998-2018 Epic Games, Inc. All Rights Reserved.
+// Copyright 1998-2019 Epic Games, Inc. All Rights Reserved.
 
 #include "NativeJSScripting.h"
 
diff --git a/Engine/Source/Runtime/WebBrowser/Private/Native/NativeJSScripting.h b/Engine/Source/Runtime/WebBrowser/Private/Native/NativeJSScripting.h
--- a/Engine/Source/Runtime/WebBrowser/Private/Native/NativeJSScripting.h
+++ b/Engine/Source/Runtime/WebBrowser/Private/Native/NativeJSScripting.h
@@ -1,4 +1,4 @@
-// Copyright 1998-2018 Epic Games, Inc. All Rights Reserved.
+// Copyright 1998-2019 Epic Games, Inc. All Rights Reserved.
 #pragma once
 
 #include "CoreMinimal.h"
diff --git a/Engine/Source/Runtime/WebBrowser/Private/Native/NativeJSStructDeserializerBackend.cpp b/Engine/Source/Runtime/WebBrowser/Private/Native/NativeJSStructDeserializerBackend.cpp
--- a/Engine/Source/Runtime/WebBrowser/Private/Native/NativeJSStructDeserializerBackend.cpp
+++ b/Engine/Source/Runtime/WebBrowser/Private/Native/NativeJSStructDeserializerBackend.cpp
@@ -1,4 +1,4 @@
-// Copyright 1998-2018 Epic Games, Inc. All Rights Reserved.
+// Copyright 1998-2019 Epic Games, Inc. All Rights Reserved.
 
 #include "NativeJSStructDeserializerBackend.h"
 #include "NativeJSScripting.h"
diff --git a/Engine/Source/Runtime/WebBrowser/Private/Native/NativeJSStructDeserializerBackend.h b/Engine/Source/Runtime/WebBrowser/Private/Native/NativeJSStructDeserializerBackend.h
--- a/Engine/Source/Runtime/WebBrowser/Private/Native/NativeJSStructDeserializerBackend.h
+++ b/Engine/Source/Runtime/WebBrowser/Private/Native/NativeJSStructDeserializerBackend.h
@@ -1,4 +1,4 @@
-// Copyright 1998-2018 Epic Games, Inc. All Rights Reserved.
+// Copyright 1998-2019 Epic Games, Inc. All Rights Reserved.
 
 #pragma once
 
diff --git a/Engine/Source/Runtime/WebBrowser/Private/Native/NativeJSStructSerializerBackend.cpp b/Engine/Source/Runtime/WebBrowser/Private/Native/NativeJSStructSerializerBackend.cpp
--- a/Engine/Source/Runtime/WebBrowser/Private/Native/NativeJSStructSerializerBackend.cpp
+++ b/Engine/Source/Runtime/WebBrowser/Private/Native/NativeJSStructSerializerBackend.cpp
@@ -1,4 +1,4 @@
-// Copyright 1998-2018 Epic Games, Inc. All Rights Reserved.
+// Copyright 1998-2019 Epic Games, Inc. All Rights Reserved.
 
 #include "Native/NativeJSStructSerializerBackend.h"
 
diff --git a/Engine/Source/Runtime/WebBrowser/Private/Native/NativeJSStructSerializerBackend.h b/Engine/Source/Runtime/WebBrowser/Private/Native/NativeJSStructSerializerBackend.h
--- a/Engine/Source/Runtime/WebBrowser/Private/Native/NativeJSStructSerializerBackend.h
+++ b/Engine/Source/Runtime/WebBrowser/Private/Native/NativeJSStructSerializerBackend.h
@@ -1,4 +1,4 @@
-// Copyright 1998-2018 Epic Games, Inc. All Rights Reserved.
+// Copyright 1998-2019 Epic Games, Inc. All Rights Reserved.
 
 #pragma once
 
diff --git a/Engine/Source/Runtime/WebBrowser/Private/Native/NativeWebBrowserProxy.cpp b/Engine/Source/Runtime/WebBrowser/Private/Native/NativeWebBrowserProxy.cpp
--- a/Engine/Source/Runtime/WebBrowser/Private/Native/NativeWebBrowserProxy.cpp
+++ b/Engine/Source/Runtime/WebBrowser/Private/Native/NativeWebBrowserProxy.cpp
@@ -1,4 +1,4 @@
-// Copyright 1998-2018 Epic Games, Inc. All Rights Reserved.
+// Copyright 1998-2019 Epic Games, Inc. All Rights Reserved.
 
 
 #include "NativeWebBrowserProxy.h"
diff --git a/Engine/Source/Runtime/WebBrowser/Private/Native/NativeWebBrowserProxy.h b/Engine/Source/Runtime/WebBrowser/Private/Native/NativeWebBrowserProxy.h
--- a/Engine/Source/Runtime/WebBrowser/Private/Native/NativeWebBrowserProxy.h
+++ b/Engine/Source/Runtime/WebBrowser/Private/Native/NativeWebBrowserProxy.h
@@ -1,4 +1,4 @@
-// Copyright 1998-2018 Epic Games, Inc. All Rights Reserved.
+// Copyright 1998-2019 Epic Games, Inc. All Rights Reserved.
 #pragma once
 
 #include "CoreMinimal.h"
diff --git a/Engine/Source/Runtime/Windows/WindowsPlatformFeatures/Private/HighlightRecorder.cpp b/Engine/Source/Runtime/Windows/WindowsPlatformFeatures/Private/HighlightRecorder.cpp
--- a/Engine/Source/Runtime/Windows/WindowsPlatformFeatures/Private/HighlightRecorder.cpp
+++ b/Engine/Source/Runtime/Windows/WindowsPlatformFeatures/Private/HighlightRecorder.cpp
@@ -1,4 +1,4 @@
-// Copyright 1998-2018 Epic Games, Inc. All Rights Reserved.
+// Copyright 1998-2019 Epic Games, Inc. All Rights Reserved.
 
 #include "HighlightRecorder.h"
 #include "WmfMp4Writer.h"
diff --git a/Engine/Source/Runtime/Windows/WindowsPlatformFeatures/Private/HighlightRecorder.h b/Engine/Source/Runtime/Windows/WindowsPlatformFeatures/Private/HighlightRecorder.h
--- a/Engine/Source/Runtime/Windows/WindowsPlatformFeatures/Private/HighlightRecorder.h
+++ b/Engine/Source/Runtime/Windows/WindowsPlatformFeatures/Private/HighlightRecorder.h
@@ -1,4 +1,4 @@
-// Copyright 1998-2018 Epic Games, Inc. All Rights Reserved.
+// Copyright 1998-2019 Epic Games, Inc. All Rights Reserved.
 
 #pragma once
 
diff --git a/Engine/Source/Runtime/Windows/WindowsPlatformFeatures/Private/WindowsPlatformFeatures.cpp b/Engine/Source/Runtime/Windows/WindowsPlatformFeatures/Private/WindowsPlatformFeatures.cpp
--- a/Engine/Source/Runtime/Windows/WindowsPlatformFeatures/Private/WindowsPlatformFeatures.cpp
+++ b/Engine/Source/Runtime/Windows/WindowsPlatformFeatures/Private/WindowsPlatformFeatures.cpp
@@ -1,4 +1,4 @@
-// Copyright 1998-2018 Epic Games, Inc. All Rights Reserved.
+// Copyright 1998-2019 Epic Games, Inc. All Rights Reserved.
 
 #include "WindowsPlatformFeatures.h"
 #include "WmfPrivate.h"
diff --git a/Engine/Source/Runtime/Windows/WindowsPlatformFeatures/Private/WindowsPlatformFeatures.h b/Engine/Source/Runtime/Windows/WindowsPlatformFeatures/Private/WindowsPlatformFeatures.h
--- a/Engine/Source/Runtime/Windows/WindowsPlatformFeatures/Private/WindowsPlatformFeatures.h
+++ b/Engine/Source/Runtime/Windows/WindowsPlatformFeatures/Private/WindowsPlatformFeatures.h
@@ -1,4 +1,4 @@
-// Copyright 1998-2018 Epic Games, Inc. All Rights Reserved.
+// Copyright 1998-2019 Epic Games, Inc. All Rights Reserved.
 
 #pragma once
 
diff --git a/Engine/Source/Runtime/Windows/WindowsPlatformFeatures/Private/WindowsPlatformFeaturesCommon.h b/Engine/Source/Runtime/Windows/WindowsPlatformFeatures/Private/WindowsPlatformFeaturesCommon.h
--- a/Engine/Source/Runtime/Windows/WindowsPlatformFeatures/Private/WindowsPlatformFeaturesCommon.h
+++ b/Engine/Source/Runtime/Windows/WindowsPlatformFeatures/Private/WindowsPlatformFeaturesCommon.h
@@ -1,4 +1,4 @@
-// Copyright 1998-2018 Epic Games, Inc. All Rights Reserved.
+// Copyright 1998-2019 Epic Games, Inc. All Rights Reserved.
 
 #pragma once
 
diff --git a/Engine/Source/Runtime/Windows/WindowsPlatformFeatures/Private/WindowsVideoRecordingSystem.cpp b/Engine/Source/Runtime/Windows/WindowsPlatformFeatures/Private/WindowsVideoRecordingSystem.cpp
--- a/Engine/Source/Runtime/Windows/WindowsPlatformFeatures/Private/WindowsVideoRecordingSystem.cpp
+++ b/Engine/Source/Runtime/Windows/WindowsPlatformFeatures/Private/WindowsVideoRecordingSystem.cpp
@@ -1,4 +1,4 @@
-// Copyright 1998-2018 Epic Games, Inc. All Rights Reserved.
+// Copyright 1998-2019 Epic Games, Inc. All Rights Reserved.
 #include "WindowsVideoRecordingSystem.h"
 #include "Modules/ModuleManager.h"
 #include "Engine/GameEngine.h"
diff --git a/Engine/Source/Runtime/Windows/WindowsPlatformFeatures/Private/WindowsVideoRecordingSystem.h b/Engine/Source/Runtime/Windows/WindowsPlatformFeatures/Private/WindowsVideoRecordingSystem.h
--- a/Engine/Source/Runtime/Windows/WindowsPlatformFeatures/Private/WindowsVideoRecordingSystem.h
+++ b/Engine/Source/Runtime/Windows/WindowsPlatformFeatures/Private/WindowsVideoRecordingSystem.h
@@ -1,4 +1,4 @@
-// Copyright 1998-2018 Epic Games, Inc. All Rights Reserved.
+// Copyright 1998-2019 Epic Games, Inc. All Rights Reserved.
 
 #pragma once
 
diff --git a/Engine/Source/Runtime/Windows/WindowsPlatformFeatures/Private/WmfMp4Writer.cpp b/Engine/Source/Runtime/Windows/WindowsPlatformFeatures/Private/WmfMp4Writer.cpp
--- a/Engine/Source/Runtime/Windows/WindowsPlatformFeatures/Private/WmfMp4Writer.cpp
+++ b/Engine/Source/Runtime/Windows/WindowsPlatformFeatures/Private/WmfMp4Writer.cpp
@@ -1,4 +1,4 @@
-// Copyright 1998-2018 Epic Games, Inc. All Rights Reserved.
+// Copyright 1998-2019 Epic Games, Inc. All Rights Reserved.
 
 #include "WmfMp4Writer.h"
 
diff --git a/Engine/Source/Runtime/Windows/WindowsPlatformFeatures/Private/WmfMp4Writer.h b/Engine/Source/Runtime/Windows/WindowsPlatformFeatures/Private/WmfMp4Writer.h
--- a/Engine/Source/Runtime/Windows/WindowsPlatformFeatures/Private/WmfMp4Writer.h
+++ b/Engine/Source/Runtime/Windows/WindowsPlatformFeatures/Private/WmfMp4Writer.h
@@ -1,4 +1,4 @@
-// Copyright 1998-2018 Epic Games, Inc. All Rights Reserved.
+// Copyright 1998-2019 Epic Games, Inc. All Rights Reserved.
 
 #pragma once
 
diff --git a/Engine/Source/Runtime/Windows/WindowsPlatformFeatures/Private/WmfPrivate.h b/Engine/Source/Runtime/Windows/WindowsPlatformFeatures/Private/WmfPrivate.h
--- a/Engine/Source/Runtime/Windows/WindowsPlatformFeatures/Private/WmfPrivate.h
+++ b/Engine/Source/Runtime/Windows/WindowsPlatformFeatures/Private/WmfPrivate.h
@@ -1,4 +1,4 @@
-// Copyright 1998-2018 Epic Games, Inc. All Rights Reserved.
+// Copyright 1998-2019 Epic Games, Inc. All Rights Reserved.
 
 #pragma once
 
diff --git a/Engine/Source/Runtime/Windows/WindowsPlatformFeatures/Private/WmfRingBuffer.cpp b/Engine/Source/Runtime/Windows/WindowsPlatformFeatures/Private/WmfRingBuffer.cpp
--- a/Engine/Source/Runtime/Windows/WindowsPlatformFeatures/Private/WmfRingBuffer.cpp
+++ b/Engine/Source/Runtime/Windows/WindowsPlatformFeatures/Private/WmfRingBuffer.cpp
@@ -1,4 +1,4 @@
-// Copyright 1998-2018 Epic Games, Inc. All Rights Reserved.
+// Copyright 1998-2019 Epic Games, Inc. All Rights Reserved.
 
 #include "WmfRingBuffer.h"
 #include "Misc/ScopeLock.h"
diff --git a/Engine/Source/Runtime/Windows/WindowsPlatformFeatures/Private/WmfRingBuffer.h b/Engine/Source/Runtime/Windows/WindowsPlatformFeatures/Private/WmfRingBuffer.h
--- a/Engine/Source/Runtime/Windows/WindowsPlatformFeatures/Private/WmfRingBuffer.h
+++ b/Engine/Source/Runtime/Windows/WindowsPlatformFeatures/Private/WmfRingBuffer.h
@@ -1,4 +1,4 @@
-// Copyright 1998-2018 Epic Games, Inc. All Rights Reserved.
+// Copyright 1998-2019 Epic Games, Inc. All Rights Reserved.
 
 #pragma once
 
diff --git a/Engine/Source/Runtime/Windows/WindowsPlatformFeatures/WindowsPlatformFeatures.Build.cs b/Engine/Source/Runtime/Windows/WindowsPlatformFeatures/WindowsPlatformFeatures.Build.cs
--- a/Engine/Source/Runtime/Windows/WindowsPlatformFeatures/WindowsPlatformFeatures.Build.cs
+++ b/Engine/Source/Runtime/Windows/WindowsPlatformFeatures/WindowsPlatformFeatures.Build.cs
@@ -1,4 +1,4 @@
-// Copyright 1998-2018 Epic Games, Inc. All Rights Reserved.
+// Copyright 1998-2019 Epic Games, Inc. All Rights Reserved.
 
 using System.IO;
 using UnrealBuildTool;
