NvEnc SDK's ReadMe.txt, so we can easily see the driver requirements.

diff --git a/Engine/Source/Runtime/GameplayMediaEncoder/Private/Microsoft/Windows/NvEncoder/ReadMe.txt b/Engine/Source/Runtime/GameplayMediaEncoder/Private/Microsoft/Windows/NvEncoder/ReadMe.txt
new file mode 100644
--- /dev/null
+++ b/Engine/Source/Runtime/GameplayMediaEncoder/Private/Microsoft/Windows/NvEncoder/ReadMe.txt
@@ -0,0 +1,48 @@
+NVIDIA Video Codec SDK 8.0 and NVIDIA Video Codec SDK 8.1 Readme and Getting Started Guide
+
+System Requirements
+
+* NVIDIA Kepler/Maxwell/Pascal/Volta GPU with hardware video accelerators - Refer to the NVIDIA Video SDK developer zone web page (https://developer.nvidia.com/nvidia-video-codec-sdk) for GPUs which support encoding and decoding acceleration.
+
+Video Codec SDK 8.0
+   * Windows: Driver version 378.66 or higher
+   * Linux:   Driver version 378.13 or higher
+   * CUDA 8.0 Toolkit (optional)
+
+Video Codec SDK 8.1
+   * Windows: Driver version 390.77  or higher
+   * Linux:   Driver version 390.25 or higher
+   * CUDA 8.0 Toolkit 
+
+
+[Windows Configuration Requirements]
+- DirectX SDK is needed. You can download the latest SDK from Microsoft's DirectX website
+- The CUDA tool kit is needed to compile the decode samples in SDK 8.1.
+- CUDA tool kit is also used for building CUDA kernels that can interop with NVENC.
+
+The following environment variables need to be set to build the sample applications included with the SDK
+* For Windows
+  - DXSDK_DIR: pointing to the DirectX SDK root directory. 
+  - The CUDA 8.0 Toolkit is optional to install if the client has Video Codec SDK 8.0. However it is mandatory if client has Video Codec SDK 8.1 on his/her machine. 
+
+[Linux Configuration Requirements]    
+* For Linux
+  - X11 and OpenGL, GLUT, GLEW libraries for video playback and display 
+  - The CUDA 8.0 Toolkit is optional to install if the client has Video Codec SDK 8.0. 
+  - CUDA 8.0 Toolkit is mandatory if client has Video Codec SDK 8.1 on his/her machine. 
+  - CUDA toolkit is also used for building CUDA kernels that can interop with NVENC. 
+  - Libraries and headers from the FFmpeg project which can be downloaded and installed
+    using the distribution's package manager or compiled from source. The sample applications
+    have been compiled and tested against the libraries and headers from FFmpeg-3.4.1.
+    The source code of FFmpeg-3.4.1 has been included in this SDK package.
+  - To build/use sample applications that depend on FFmpeg, users may need to
+      * Add the directory (/usr/local/lib/pkgconfig by default) to the PKG_CONFIG_PATH environment variable.
+        This is required by the Makefile to determine the include paths for the FFmpeg headers.
+      * Add the directory where the FFmpeg libraries are installed, to the LD_LIBRARY_PATH environment variable.
+        This is required for resolving runtime dependencies on FFmpeg libraries.
+
+[Common to all OS platforms]
+* To download the CUDA 8.0 toolkit, please go to the following web site:
+  http://developer.nvidia.com/cuda/cuda-toolkit
+
+
