Fixed audio distortion when recording gameplay with GameplayMediaEncoder.  It was caused by loud sounds wrapping around.

diff --git a/Engine/Source/Runtime/GameplayMediaEncoder/Private/GameplayMediaEncoder.cpp b/Engine/Source/Runtime/GameplayMediaEncoder/Private/GameplayMediaEncoder.cpp
--- a/Engine/Source/Runtime/GameplayMediaEncoder/Private/GameplayMediaEncoder.cpp
+++ b/Engine/Source/Runtime/GameplayMediaEncoder/Private/GameplayMediaEncoder.cpp
@@ -251,6 +251,7 @@ bool FGameplayMediaEncoder::Initialize()
 #endif
 
 	StartTime = 0;
+	AudioStartTime = 0;
 
 	MemoryCheckpoint("Video encoder initialized");
 
@@ -282,7 +283,8 @@ bool FGameplayMediaEncoder::Start()
 	}
 
 	StartTime = FTimespan::FromSeconds(FPlatformTime::Seconds());
-	LastAudioInputTimestamp = LastVideoInputTimestamp = GetMediaTimestamp();
+	// We set this to 0, and the first audio buffer we get will be considered the first one
+	AudioStartTime = 0;
 	NumCapturedFrames = 0;
 
 	//
@@ -327,6 +329,7 @@ void FGameplayMediaEncoder::Stop()
 
 	VideoEncoder->Stop();
 	StartTime = 0;
+	AudioStartTime = 0;
 }
 
 void FGameplayMediaEncoder::Shutdown()
@@ -382,7 +385,7 @@ void FGameplayMediaEncoder::OnNewSubmixBuffer(const USoundSubmix* OwningSubmix,
 		return;
 	}
 
-	ProcessAudioFrame(AudioData, NumSamples, NumChannels, SampleRate);
+	ProcessAudioFrame(AudioData, NumSamples, NumChannels, SampleRate, AudioClock);
 }
 
 void FGameplayMediaEncoder::OnBackBufferReady(SWindow& SlateWindow, const FTexture2DRHIRef& BackBuffer)
@@ -392,7 +395,7 @@ void FGameplayMediaEncoder::OnBackBufferReady(SWindow& SlateWindow, const FTextu
 	ProcessVideoFrame(BackBuffer);
 }
 
-bool FGameplayMediaEncoder::ProcessAudioFrame(const float* AudioData, int32 NumSamples, int32 NumChannels, int32 SampleRate)
+bool FGameplayMediaEncoder::ProcessAudioFrame(const float* AudioData, int32 NumSamples, int32 NumChannels, int32 SampleRate, double AudioClock)
 {
 	Audio::AlignedFloatBuffer InData;
 	InData.Append(AudioData, NumSamples);
@@ -404,14 +407,22 @@ bool FGameplayMediaEncoder::ProcessAudioFrame(const float* AudioData, int32 NumS
 		FloatBuffer.MixBufferToChannels(HardcodedAudioNumChannels);
 	}
 
+	// First we need to clamp to [-1.0, 1.0], because the conversion to 16 bits doesn't clamp it,
+	// and we can end up with wrapping when playing loud sounds.
+	FloatBuffer.Clamp();
 	PCM16 = FloatBuffer;
 
 	{
 		FScopeLock Lock(&AudioProcessingCS);
 
-		FTimespan Timestamp = GetMediaTimestamp();
-		FTimespan Duration = Timestamp - LastAudioInputTimestamp;
-		LastAudioInputTimestamp = Timestamp;
+
+		if (FMath::IsNearlyEqual(AudioStartTime, double(0)))
+		{
+			AudioStartTime = AudioClock;
+		}
+
+		FTimespan Timestamp = FTimespan::FromSeconds(AudioClock - AudioStartTime);
+		FTimespan Duration = FTimespan::FromSeconds((double)(NumSamples/NumChannels) / SampleRate);
 		AudioEncoder->Process(reinterpret_cast<const int8*>(PCM16.GetData()), PCM16.GetNumSamples() * sizeof(*PCM16.GetData()), Timestamp, Duration);
 	}
 
@@ -492,4 +503,3 @@ bool FGameplayMediaEncoder::ChangeVideoConfig()
 }
 
 GAMEPLAYMEDIAENCODER_END
-
diff --git a/Engine/Source/Runtime/GameplayMediaEncoder/Public/GameplayMediaEncoder.h b/Engine/Source/Runtime/GameplayMediaEncoder/Public/GameplayMediaEncoder.h
--- a/Engine/Source/Runtime/GameplayMediaEncoder/Public/GameplayMediaEncoder.h
+++ b/Engine/Source/Runtime/GameplayMediaEncoder/Public/GameplayMediaEncoder.h
@@ -81,7 +81,7 @@ private:
 	// ISubmixBufferListener interface
 	void OnNewSubmixBuffer(const USoundSubmix* OwningSubmix, float* AudioData, int32 NumSamples, int32 NumChannels, const int32 SampleRate, double AudioClock) override;
 
-	bool ProcessAudioFrame(const float* AudioData, int32 NumSamples, int32 NumChannels, int32 SampleRate);
+	bool ProcessAudioFrame(const float* AudioData, int32 NumSamples, int32 NumChannels, int32 SampleRate, double AudioClock);
 	bool ProcessVideoFrame(const FTexture2DRHIRef& BackBuffer);
 
 	bool ChangeVideoConfig();
@@ -105,7 +105,14 @@ private:
 
 	uint64 NumCapturedFrames = 0;
 	FTimespan StartTime = 0;
-	FTimespan LastAudioInputTimestamp = 0;
+	// For audio, we can't use anything timed with FPlatformMisc::Seconds(), because we need
+	// to follow the audio clock.
+	// In other words, the pace of the audio clock exactly depends on the amount of audio data we
+	// get.
+	// By using FPlatform::Seconds(), audio timing wouldn't exactly match the duration or timestamps of
+	// the audio buffers from the audio device perspective.
+	double AudioStartTime = 0;
+
 	FTimespan LastVideoInputTimestamp = 0;
 
 	// It is possible to suspend the processing of media samples which is
