removed all `FThread` definitions scattered around as it's now part of Core and was causing "already defined" error on Xbox

diff --git a/Engine/Plugins/Experimental/PixelStreaming/Source/PixelStreaming/Private/PixelStreamingNvVideoEncoder.cpp b/Engine/Plugins/Experimental/PixelStreaming/Source/PixelStreaming/Private/PixelStreamingNvVideoEncoder.cpp
--- a/Engine/Plugins/Experimental/PixelStreaming/Source/PixelStreaming/Private/PixelStreamingNvVideoEncoder.cpp
+++ b/Engine/Plugins/Experimental/PixelStreaming/Source/PixelStreaming/Private/PixelStreamingNvVideoEncoder.cpp
@@ -9,8 +9,7 @@
 #include "RendererInterface.h"
 #include "RHIStaticStates.h"
 #include "PipelineStateCache.h"
-#include "HAL/Runnable.h"
-#include "HAL/RunnableThread.h"
+#include "HAL/Thread.h"
 #include "HAL/ThreadSafeBool.h"
 #include "HAL/ThreadSafeCounter.h"
 #include "Misc/CoreDelegates.h"
diff --git a/Engine/Plugins/Experimental/PixelStreaming/Source/PixelStreaming/Private/ProxyConnection.h b/Engine/Plugins/Experimental/PixelStreaming/Source/PixelStreaming/Private/ProxyConnection.h
--- a/Engine/Plugins/Experimental/PixelStreaming/Source/PixelStreaming/Private/ProxyConnection.h
+++ b/Engine/Plugins/Experimental/PixelStreaming/Source/PixelStreaming/Private/ProxyConnection.h
@@ -3,6 +3,7 @@
 #pragma once
 
 #include "HAL/ThreadSafeBool.h"
+#include "HAL/Thread.h"
 #include "Misc/ScopeLock.h"
 #include "Templates/SharedPointer.h"
 
diff --git a/Engine/Plugins/Experimental/PixelStreaming/Source/PixelStreaming/Private/Utils.h b/Engine/Plugins/Experimental/PixelStreaming/Source/PixelStreaming/Private/Utils.h
--- a/Engine/Plugins/Experimental/PixelStreaming/Source/PixelStreaming/Private/Utils.h
+++ b/Engine/Plugins/Experimental/PixelStreaming/Source/PixelStreaming/Private/Utils.h
@@ -2,41 +2,8 @@
 
 #pragma once
 
-#include "HAL/Runnable.h"
-#include "HAL/RunnableThread.h"
 #include <chrono>
 
-class FThread final : public FRunnable
-{
-public:
-	using FCallback = TFunction<void()>;
-
-	explicit FThread(TCHAR const* ThreadName, const FCallback& Callback) :
-		Callback(Callback)
-	{
-		Thread = FRunnableThread::Create(this, ThreadName);
-	}
-
-	void Join()
-	{
-		Thread->WaitForCompletion();
-	}
-
-	virtual uint32 Run() override
-	{
-		Callback();
-		return 0;
-	}
-
-private:
-	FCallback Callback;
-	FRunnableThread* Thread;
-
-private:
-	FThread(const FThread&) = delete;
-	FThread& operator=(const FThread&) = delete;
-};
-
 // uses chrono library to have comparable timestamps between UE4 and webrtc app
 inline uint64 NowMs()
 {
diff --git a/Engine/Source/Runtime/GameplayMediaEncoder/Private/GameplayMediaEncoderCommon.h b/Engine/Source/Runtime/GameplayMediaEncoder/Private/GameplayMediaEncoderCommon.h
--- a/Engine/Source/Runtime/GameplayMediaEncoder/Private/GameplayMediaEncoderCommon.h
+++ b/Engine/Source/Runtime/GameplayMediaEncoder/Private/GameplayMediaEncoderCommon.h
@@ -8,8 +8,6 @@
 #include "RHIStaticStates.h"
 #include "HAL/ThreadSafeBool.h"
 #include "HAL/ThreadSafeCounter.h"
-#include "HAL/Runnable.h"
-#include "HAL/RunnableThread.h"
 #include "HAL/Event.h"
 #include "Misc/ScopeExit.h"
 #include "ShaderCore.h"
@@ -244,38 +242,6 @@ ID3D12Device* GetUE4DxDevice();
 
 #endif
 
-class FThread final : public FRunnable
-{
-public:
-	using FCallback = TFunction<void()>;
-
-	explicit FThread(TCHAR const* ThreadName, const FCallback& Callback) :
-		Callback(Callback)
-	{
-		Thread.Reset(FRunnableThread::Create(this, ThreadName, TPri_BelowNormal));
-	}
-
-	void Join()
-	{
-		Thread->WaitForCompletion();
-	}
-
-	virtual uint32 Run() override
-	{
-		Callback();
-		return 0;
-	}
-
-private:
-	FCallback Callback;
-	TUniquePtr<FRunnableThread> Thread;
-
-private:
-	FThread(const FThread&) = delete;
-	FThread& operator=(const FThread&) = delete;
-};
-
-
 template<typename F>
 inline void ExecuteRHICommand(F&& Functor)
 {
diff --git a/Engine/Source/Runtime/GameplayMediaEncoder/Private/Microsoft/Windows/NvVideoEncoder.cpp b/Engine/Source/Runtime/GameplayMediaEncoder/Private/Microsoft/Windows/NvVideoEncoder.cpp
--- a/Engine/Source/Runtime/GameplayMediaEncoder/Private/Microsoft/Windows/NvVideoEncoder.cpp
+++ b/Engine/Source/Runtime/GameplayMediaEncoder/Private/Microsoft/Windows/NvVideoEncoder.cpp
@@ -11,6 +11,7 @@
 #include "PipelineStateCache.h"
 #include "CommonRenderResources.h"
 #include "Modules/ModuleManager.h"
+#include "HAL/Thread.h"
 
 #include "SceneUtils.h"
 #include "EncoderDevice.h"
diff --git a/Engine/Source/Runtime/Windows/WindowsPlatformFeatures/Private/HighlightRecorder.cpp b/Engine/Source/Runtime/Windows/WindowsPlatformFeatures/Private/HighlightRecorder.cpp
--- a/Engine/Source/Runtime/Windows/WindowsPlatformFeatures/Private/HighlightRecorder.cpp
+++ b/Engine/Source/Runtime/Windows/WindowsPlatformFeatures/Private/HighlightRecorder.cpp
@@ -14,6 +14,7 @@
 
 #include "HAL/PlatformTime.h"
 #include "HAL/PlatformFilemanager.h"
+#include "HAL/Thread.h"
 #include "UnrealEngine.h"
 #include "VideoRecordingSystem.h"
 
diff --git a/Engine/Source/Runtime/Windows/WindowsPlatformFeatures/Private/HighlightRecorder.h b/Engine/Source/Runtime/Windows/WindowsPlatformFeatures/Private/HighlightRecorder.h
--- a/Engine/Source/Runtime/Windows/WindowsPlatformFeatures/Private/HighlightRecorder.h
+++ b/Engine/Source/Runtime/Windows/WindowsPlatformFeatures/Private/HighlightRecorder.h
@@ -16,40 +16,10 @@
 #include "GameplayMediaEncoder.h"
 
 class FWmfMp4Writer;
+class FThread;
 
 DECLARE_LOG_CATEGORY_EXTERN(HighlightRecorder, Log, VeryVerbose);
 
-class FThread final : public FRunnable
-{
-public:
-	using FCallback = TFunction<void()>;
-
-	explicit FThread(TCHAR const* ThreadName, const FCallback& Callback) :
-		Callback(Callback)
-	{
-		Thread.Reset(FRunnableThread::Create(this, ThreadName, TPri_BelowNormal));
-	}
-
-	void Join()
-	{
-		Thread->WaitForCompletion();
-	}
-
-	virtual uint32 Run() override
-	{
-		Callback();
-		return 0;
-	}
-
-private:
-	FCallback Callback;
-	TUniquePtr<FRunnableThread> Thread;
-
-private:
-	FThread(const FThread&) = delete;
-	FThread& operator=(const FThread&) = delete;
-};
-
 class FHighlightRecorder final : private IGameplayMediaEncoderListener
 {
 public:
