Fixed the null audio device pace.
Improved GameplayMediaEncoder audio handling (timestamps and duration). It now also adjust the audio clock if it falls behind, such as if we break into the debugger.

diff --git a/Engine/Source/Runtime/AudioMixerCore/Private/AudioMixerNullDevice.cpp b/Engine/Source/Runtime/AudioMixerCore/Private/AudioMixerNullDevice.cpp
--- a/Engine/Source/Runtime/AudioMixerCore/Private/AudioMixerNullDevice.cpp
+++ b/Engine/Source/Runtime/AudioMixerCore/Private/AudioMixerNullDevice.cpp
@@ -8,10 +8,27 @@ namespace Audio
 {
 	uint32 FMixerNullCallback::Run()
 	{
+		//
+		// To simulate an audio device requesting for more audio, we sleep between callbacks.
+		// The problem with this is that Sleep is not accurate. It will always be slightly higher than requested,
+		// which means that audio will be generated slightly slower than the stated sample rate.
+		// To correct this, we keep track of the real time passed, and adjust the sleep time accordingly so the audio clock
+		// stays as close to the real time clock as possible.
+
+		double AudioClock = FPlatformTime::Seconds();
+
+		float SleepTime = CallbackTime; 
 		while (!bShouldShutdown)
 		{
 			Callback();
-			FPlatformProcess::Sleep(CallbackTime);
+			FPlatformProcess::Sleep(FMath::Max(0.f, SleepTime));
+
+			AudioClock += CallbackTime;
+			double RealClock = FPlatformTime::Seconds();
+			double AudioVsReal = RealClock - AudioClock;
+			// For the next sleep, we adjust the sleep duration to try and keep the audio clock as close
+			// to the real time clock as possible
+			SleepTime = CallbackTime - AudioVsReal;
 		}
 
 		return 0;
diff --git a/Engine/Source/Runtime/GameplayMediaEncoder/Private/GameplayMediaEncoder.cpp b/Engine/Source/Runtime/GameplayMediaEncoder/Private/GameplayMediaEncoder.cpp
--- a/Engine/Source/Runtime/GameplayMediaEncoder/Private/GameplayMediaEncoder.cpp
+++ b/Engine/Source/Runtime/GameplayMediaEncoder/Private/GameplayMediaEncoder.cpp
@@ -251,7 +251,7 @@ bool FGameplayMediaEncoder::Initialize()
 #endif
 
 	StartTime = 0;
-	AudioStartTime = 0;
+	AudioClock = 0;
 
 	MemoryCheckpoint("Video encoder initialized");
 
@@ -283,8 +283,7 @@ bool FGameplayMediaEncoder::Start()
 	}
 
 	StartTime = FTimespan::FromSeconds(FPlatformTime::Seconds());
-	// We set this to 0, and the first audio buffer we get will be considered the first one
-	AudioStartTime = 0;
+	AudioClock = 0;
 	NumCapturedFrames = 0;
 
 	//
@@ -329,7 +328,7 @@ void FGameplayMediaEncoder::Stop()
 
 	VideoEncoder->Stop();
 	StartTime = 0;
-	AudioStartTime = 0;
+	AudioClock = 0;
 }
 
 void FGameplayMediaEncoder::Shutdown()
@@ -385,7 +384,7 @@ void FGameplayMediaEncoder::OnNewSubmixBuffer(const USoundSubmix* OwningSubmix,
 		return;
 	}
 
-	ProcessAudioFrame(AudioData, NumSamples, NumChannels, SampleRate, AudioClock);
+	ProcessAudioFrame(AudioData, NumSamples, NumChannels, SampleRate);
 }
 
 void FGameplayMediaEncoder::OnBackBufferReady(SWindow& SlateWindow, const FTexture2DRHIRef& BackBuffer)
@@ -395,7 +394,7 @@ void FGameplayMediaEncoder::OnBackBufferReady(SWindow& SlateWindow, const FTextu
 	ProcessVideoFrame(BackBuffer);
 }
 
-bool FGameplayMediaEncoder::ProcessAudioFrame(const float* AudioData, int32 NumSamples, int32 NumChannels, int32 SampleRate, double AudioClock)
+bool FGameplayMediaEncoder::ProcessAudioFrame(const float* AudioData, int32 NumSamples, int32 NumChannels, int32 SampleRate)
 {
 	Audio::AlignedFloatBuffer InData;
 	InData.Append(AudioData, NumSamples);
@@ -412,20 +411,24 @@ bool FGameplayMediaEncoder::ProcessAudioFrame(const float* AudioData, int32 NumS
 	FloatBuffer.Clamp();
 	PCM16 = FloatBuffer;
 
+	// Adjust the AudioClock if for some reason it falls behind real time. This can happen if the game spikes, or if we break into the debugger.
+	FTimespan Now = GetMediaTimestamp();
+	if (AudioClock < Now.GetTotalSeconds())
 	{
-		FScopeLock Lock(&AudioProcessingCS);
-
-
-		if (FMath::IsNearlyEqual(AudioStartTime, double(0)))
-		{
-			AudioStartTime = AudioClock;
-		}
+		UE_LOG(GameplayMediaEncoder, Warning, TEXT("Audio clock falling behind real time clock by %.3f seconds. Ajusting audio clock"), Now.GetTotalSeconds()-AudioClock);
+		// Put it slightly ahead of the real time clock
+		AudioClock = Now.GetTotalSeconds() + (FloatBuffer.GetSampleDuration() / 2);
+	}
 
-		FTimespan Timestamp = FTimespan::FromSeconds(AudioClock - AudioStartTime);
-		FTimespan Duration = FTimespan::FromSeconds((double)(NumSamples/NumChannels) / SampleRate);
+	{
+		FScopeLock Lock(&AudioProcessingCS);
+		FTimespan Timestamp = FTimespan::FromSeconds(AudioClock);
+		FTimespan Duration = FTimespan::FromSeconds(FloatBuffer.GetSampleDuration());
 		AudioEncoder->Process(reinterpret_cast<const int8*>(PCM16.GetData()), PCM16.GetNumSamples() * sizeof(*PCM16.GetData()), Timestamp, Duration);
 	}
 
+	AudioClock += FloatBuffer.GetSampleDuration();
+
 	return true;
 }
 
diff --git a/Engine/Source/Runtime/GameplayMediaEncoder/Public/GameplayMediaEncoder.h b/Engine/Source/Runtime/GameplayMediaEncoder/Public/GameplayMediaEncoder.h
--- a/Engine/Source/Runtime/GameplayMediaEncoder/Public/GameplayMediaEncoder.h
+++ b/Engine/Source/Runtime/GameplayMediaEncoder/Public/GameplayMediaEncoder.h
@@ -81,7 +81,7 @@ private:
 	// ISubmixBufferListener interface
 	void OnNewSubmixBuffer(const USoundSubmix* OwningSubmix, float* AudioData, int32 NumSamples, int32 NumChannels, const int32 SampleRate, double AudioClock) override;
 
-	bool ProcessAudioFrame(const float* AudioData, int32 NumSamples, int32 NumChannels, int32 SampleRate, double AudioClock);
+	bool ProcessAudioFrame(const float* AudioData, int32 NumSamples, int32 NumChannels, int32 SampleRate);
 	bool ProcessVideoFrame(const FTexture2DRHIRef& BackBuffer);
 
 	bool ChangeVideoConfig();
@@ -105,13 +105,10 @@ private:
 
 	uint64 NumCapturedFrames = 0;
 	FTimespan StartTime = 0;
-	// For audio, we can't use anything timed with FPlatformMisc::Seconds(), because we need
-	// to follow the audio clock.
-	// In other words, the pace of the audio clock exactly depends on the amount of audio data we
-	// get.
-	// By using FPlatform::Seconds(), audio timing wouldn't exactly match the duration or timestamps of
-	// the audio buffers from the audio device perspective.
-	double AudioStartTime = 0;
+	// Instead of using the AudioClock parameter ISubmixBufferListener::OnNewSubmixBuffer gives us, we calculate our own, by
+	// advancing it as we receive more data.
+	// This is so that we can adjust the clock if things get out of sync, such as if we break into the debugger.
+	double AudioClock = 0;
 
 	FTimespan LastVideoInputTimestamp = 0;
 
