Small improvement to avoid a compiler warning (Parameter with the same name as a member variable)

diff --git a/Engine/Source/Runtime/GameplayMediaEncoder/Private/GameplayMediaEncoder.cpp b/Engine/Source/Runtime/GameplayMediaEncoder/Private/GameplayMediaEncoder.cpp
--- a/Engine/Source/Runtime/GameplayMediaEncoder/Private/GameplayMediaEncoder.cpp
+++ b/Engine/Source/Runtime/GameplayMediaEncoder/Private/GameplayMediaEncoder.cpp
@@ -370,7 +370,7 @@ FTimespan FGameplayMediaEncoder::GetMediaTimestamp() const
 	return FTimespan::FromSeconds(FPlatformTime::Seconds()) - StartTime;
 }
 
-void FGameplayMediaEncoder::OnNewSubmixBuffer(const USoundSubmix* OwningSubmix, float* AudioData, int32 NumSamples, int32 NumChannels, const int32 SampleRate, double AudioClock)
+void FGameplayMediaEncoder::OnNewSubmixBuffer(const USoundSubmix* OwningSubmix, float* AudioData, int32 NumSamples, int32 NumChannels, const int32 SampleRate, double /*AudioClock*/)
 {
 	CSV_SCOPED_TIMING_STAT(GameplayMediaEncoder, OnNewSubmixBuffer);
 	if (SampleRate != HardcodedAudioSamplerate)
