GameplayMediaEncoder: Fixed UE4 reference.
#rnx

diff --git a/Engine/Source/Runtime/GameplayMediaEncoder/Private/GameplayMediaEncoder.cpp b/Engine/Source/Runtime/GameplayMediaEncoder/Private/GameplayMediaEncoder.cpp
--- a/Engine/Source/Runtime/GameplayMediaEncoder/Private/GameplayMediaEncoder.cpp
+++ b/Engine/Source/Runtime/GameplayMediaEncoder/Private/GameplayMediaEncoder.cpp
@@ -15,7 +15,7 @@
 DEFINE_LOG_CATEGORY(GameplayMediaEncoder);
 CSV_DEFINE_CATEGORY(GameplayMediaEncoder, true);
 
-// right now we support only 48KHz audio sample rate as it's the only config UE4 seems to output
+// right now we support only 48KHz audio sample rate as it's the only config UE seems to output
 // WMF AAC encoder supports also 44100Hz so its support can be easily added
 const uint32 HardcodedAudioSamplerate = 48000;
 // for now we downsample to stereo. WMF AAC encoder also supports 6 (5.1) channels
