- Adds Vulkan and Linux support to GameplayMediaEncoder
- Fixes to GameplayMediaEncoder crash as HardwareEncoders Plugin is needed to be enabled to use them. Added a alert msg to enable them as they are not enabled by default.

diff --git a/Engine/Source/Runtime/GameplayMediaEncoder/Private/GameplayMediaEncoder.cpp b/Engine/Source/Runtime/GameplayMediaEncoder/Private/GameplayMediaEncoder.cpp
--- a/Engine/Source/Runtime/GameplayMediaEncoder/Private/GameplayMediaEncoder.cpp
+++ b/Engine/Source/Runtime/GameplayMediaEncoder/Private/GameplayMediaEncoder.cpp
@@ -220,7 +220,7 @@ bool FGameplayMediaEncoder::Initialize()
 	{
 		FString RHIName = GDynamicRHI->GetName();
 
-#if PLATFORM_WINDOWS && PLATFORM_DESKTOP
+#if PLATFORM_DESKTOP && !PLATFORM_APPLE
 		if(RHIName == TEXT("D3D11"))
 		{
 			VideoEncoderInput = AVEncoder::FVideoEncoderInput::CreateForD3D11(GDynamicRHI->RHIGetNativeDevice(), VideoConfig.Width, VideoConfig.Height, true, IsRHIDeviceAMD());
@@ -229,22 +229,46 @@ bool FGameplayMediaEncoder::Initialize()
 		{
 			VideoEncoderInput = AVEncoder::FVideoEncoderInput::CreateForD3D12(GDynamicRHI->RHIGetNativeDevice(), VideoConfig.Width, VideoConfig.Height, true, IsRHIDeviceNVIDIA());
 		}
+		else if (RHIName == TEXT("Vulkan"))
+		{
+			VideoEncoderInput = AVEncoder::FVideoEncoderInput::CreateForVulkan(GDynamicRHI->RHIGetNativeDevice(), VideoConfig.Width, VideoConfig.Height, true);
+		}
 		else
 #endif
 		{
-			UE_LOG(GameplayMediaEncoder, Fatal, TEXT("GameplayMediaEncoder does not currently support the current Platform/RHI combo."));
-			unimplemented();
+			UE_LOG(GameplayMediaEncoder, Error, TEXT("Video encoding is not supported with the current Platform/RHI combo."));
+			return false;
+		}
+	}
+
+	const TArray<AVEncoder::FVideoEncoderInfo>& AvailableEncodersInfo = AVEncoder::FVideoEncoderFactory::Get().GetAvailable();
+
+	if (AvailableEncodersInfo.Num() == 0)
+	{
+		UE_LOG(GameplayMediaEncoder, Error, TEXT("No video encoders found. Check if relevent encoder plugins have been enabled for this project."));
+		return false;
+	}
+
+	for (const auto& EncoderInfo : AvailableEncodersInfo)
+	{
+		if (EncoderInfo.CodecType == AVEncoder::ECodecType::H264)
+		{
+			VideoEncoder = AVEncoder::FVideoEncoderFactory::Get().Create(EncoderInfo.ID, VideoEncoderInput, videoInit);
 		}
 	}
 
-	auto& Available = AVEncoder::FVideoEncoderFactory::Get().GetAvailable();
-	VideoEncoder = AVEncoder::FVideoEncoderFactory::Get().Create(Available[0].ID, VideoEncoderInput, videoInit);
+	if (!VideoEncoder)
+	{
+		UE_LOG(GameplayMediaEncoder, Error, TEXT("No H264 video encoder found. Check if relevent encoder plugins have been enabled for this project."));
+		return false;
+	}
+
 	VideoEncoder->SetOnEncodedPacket([this](uint32 LayerIndex, const AVEncoder::FVideoEncoderInputFrame* Frame, const AVEncoder::FCodecPacket& Packet)
 	                                 { OnEncodedVideoFrame(LayerIndex, Frame, Packet); });
 
 	if(!VideoEncoder)
 	{
-		UE_LOG(GameplayMediaEncoder, Error, TEXT("Could not create video encoder"));
+		UE_LOG(GameplayMediaEncoder, Fatal, TEXT("Creating video encoder failed."));
 		return false;
 	}
 
