Fix deadlock in FGameplayMediaEncoder. Exclude FGameplayMediaEncoder::Start from ListenerCS mutex as it internally locks FAudioEncoder::ListenerMutex resulting in the following deadlock.

GameThread:

FHighlightRecorder::Start() calls
FGameplayMediaEncoder::Get()->RegisterListener(this), which takes a scoped lock on
:lock: FGameplayMediaEncoder::ListenersCS. Then it calls (conditionally)
FGameplayMediaEncoder::Start(), which eventually calls
FAudioEncoder::RegisterListener(*this), which takes another lock:
:lock: FAudioEncoder::ListenersMutex

AudioMixerRenderThread:
Run()
FOutputBuffer::MixNextBuffer()
FMixerDevice::OnProcessAudioStream()
FMixerSubmix::ProcessAudio()
FMixerSubmix::SendAudioToSubmixBufferListeners()
FGameplayMediaEncoder::OnNewSubmixBuffer()
FGameplayMediaEncoder::ProcessAudioFrame()
FWmfAudioEncoder::Encode()
FAudioEncoder::OnEncodedAudioFrame(), which locks
:lock: FAudioEncoder::ListenersMutex. In that scope it calls
FGameplayMediaEncoder::OnEncodedAudioFrame() and locks
:lock: FGameplayMediaEncoder::ListenersCS, causing a deadlock with the Game thread

diff --git a/Engine/Source/Runtime/GameplayMediaEncoder/Private/GameplayMediaEncoder.cpp b/Engine/Source/Runtime/GameplayMediaEncoder/Private/GameplayMediaEncoder.cpp
--- a/Engine/Source/Runtime/GameplayMediaEncoder/Private/GameplayMediaEncoder.cpp
+++ b/Engine/Source/Runtime/GameplayMediaEncoder/Private/GameplayMediaEncoder.cpp
@@ -99,9 +99,12 @@ PRAGMA_ENABLE_DEPRECATION_WARNINGS
 bool FGameplayMediaEncoder::RegisterListener(IGameplayMediaEncoderListener* Listener)
 {
 	check(IsInGameThread());
-	FScopeLock Lock(&ListenersCS);
 
-	if(Listeners.Num() == 0)
+	ListenersCS.Lock();
+	bool bFirstListener = Listeners.Num() == 0;
+	ListenersCS.Unlock();
+	
+	if(bFirstListener)
 	{
 		UE_LOG(GameplayMediaEncoder, Log, TEXT("Registering the first listener"));
 		if(!Start())
@@ -110,6 +113,7 @@ bool FGameplayMediaEncoder::RegisterListener(IGameplayMediaEncoderListener* List
 		}
 	}
 
+	FScopeLock Lock(&ListenersCS);
 	Listeners.AddUnique(Listener);
 	return true;
 }
